- Updated to latest PVRTexToolCLI

diff --git a/Engine/Source/Developer/TextureFormatPVR/Private/TextureFormatPVR.cpp b/Engine/Source/Developer/TextureFormatPVR/Private/TextureFormatPVR.cpp
--- a/Engine/Source/Developer/TextureFormatPVR/Private/TextureFormatPVR.cpp
+++ b/Engine/Source/Developer/TextureFormatPVR/Private/TextureFormatPVR.cpp
@@ -444,7 +444,7 @@ class FTextureFormatPVR : public ITextureFormat
             *CompressionMode,
 			bIsPVRTC2 ? 2 : 4);
 #if PLATFORM_MAC
-		FString CompressorPath(FPaths::EngineDir() + TEXT("Binaries/ThirdParty/ImgTec/PVRTexToolCL"));
+		FString CompressorPath(FPaths::EngineDir() + TEXT("Binaries/ThirdParty/ImgTec/PVRTexToolCLI"));
 #else
 		FString CompressorPath(FPaths::EngineDir() + TEXT("Binaries/ThirdParty/ImgTec/PVRTexToolCLI.exe"));
 #endif
