- Added cooker settings to Project settings, first prop is PVRTC quality level

diff --git a/Engine/Source/Developer/TextureFormatPVR/Private/TextureFormatPVR.cpp b/Engine/Source/Developer/TextureFormatPVR/Private/TextureFormatPVR.cpp
--- a/Engine/Source/Developer/TextureFormatPVR/Private/TextureFormatPVR.cpp
+++ b/Engine/Source/Developer/TextureFormatPVR/Private/TextureFormatPVR.cpp
@@ -229,8 +229,10 @@ static void UseOriginal(const FImage& InImage, FCompressedImage2D& OutCompressed
 
 static FString GetPVRTCQualityString()
 {
-	// start at default quality
-	uint32 CompressionModeValue = DEFAULT_QUALITY;
+	// start at default quality, then lookup in .ini file
+	int32 CompressionModeValue = DEFAULT_QUALITY;
+	GConfig->GetInt(TEXT("/Script/UnrealEd.CookerSettings"), TEXT("DefaultPVRTCQuality"), CompressionModeValue, GEngineIni);
+
 	FParse::Value(FCommandLine::Get(), TEXT("-pvrtcquality="), CompressionModeValue);
 	CompressionModeValue = FMath::Min<uint32>(CompressionModeValue, MAX_QUALITY);
 
@@ -252,7 +254,9 @@ static FString GetPVRTCQualityString()
 static uint16 GetPVRTCQualityForVersion()
 {
 	// start at default quality
-	uint32 CompressionModeValue = DEFAULT_QUALITY;
+	int32 CompressionModeValue = DEFAULT_QUALITY;
+	GConfig->GetInt(TEXT("/Script/UnrealEd.CookerSettings"), TEXT("DefaultPVRTCQuality"), CompressionModeValue, GEngineIni);
+
 	FParse::Value(FCommandLine::Get(), TEXT("-pvrtcquality="), CompressionModeValue);
 	CompressionModeValue = FMath::Min<uint32>(CompressionModeValue, MAX_QUALITY);
 
diff --git a/Engine/Source/Editor/ProjectSettingsViewer/Private/ProjectSettingsViewerModule.cpp b/Engine/Source/Editor/ProjectSettingsViewer/Private/ProjectSettingsViewerModule.cpp
--- a/Engine/Source/Editor/ProjectSettingsViewer/Private/ProjectSettingsViewerModule.cpp
+++ b/Engine/Source/Editor/ProjectSettingsViewer/Private/ProjectSettingsViewerModule.cpp
@@ -3,6 +3,7 @@
 #include "ProjectSettingsViewerPrivatePCH.h"
 #include "Engine/Console.h"
 #include "ProjectTargetPlatformEditor.h"
+#include "Settings/CookerSettings.h"
 
 
 #define LOCTEXT_NAMESPACE "FProjectSettingsViewerModule"
@@ -145,6 +146,13 @@ protected:
 			LOCTEXT("ProjectRenderingSettingsDescription", "Rendering settings."),
 			GetMutableDefault<URendererSettings>()
 		);
+
+		// Rendering settings
+		SettingsModule.RegisterSettings("Project", "Engine", "Cooker",
+			LOCTEXT("CookerSettingsName", "Cooker"),
+			LOCTEXT("CookerSettingsDescription", "Various cooker settings."),
+			GetMutableDefault<UCookerSettings>()
+			);
 	}
 
 	/**
diff --git a/Engine/Source/Editor/UnrealEd/Classes/Settings/CookerSettings.h b/Engine/Source/Editor/UnrealEd/Classes/Settings/CookerSettings.h
new file mode 100644
--- /dev/null
+++ b/Engine/Source/Editor/UnrealEd/Classes/Settings/CookerSettings.h
@@ -0,0 +1,26 @@
+// Copyright 1998-2014 Epic Games, Inc. All Rights Reserved.
+
+/*=============================================================================
+	CookerSettings.h: Declares the UCookerSettings class.
+=============================================================================*/
+
+#pragma once
+
+#include "CookerSettings.generated.h"
+
+
+/**
+ * Implements per-project cooker settings exposed to the editor
+ */
+UCLASS(config=Engine)
+class UNREALED_API UCookerSettings
+	: public UObject
+{
+	GENERATED_UCLASS_BODY()
+
+public:
+
+	/** Quality of 0 means fastest, 4 means best */
+	UPROPERTY(config, EditAnywhere, Category=Textures)
+	int32 DefaultPVRTCQuality;
+};
diff --git a/Engine/Source/Editor/UnrealEd/Private/Commandlets/CookCommandlet.cpp b/Engine/Source/Editor/UnrealEd/Private/Commandlets/CookCommandlet.cpp
--- a/Engine/Source/Editor/UnrealEd/Private/Commandlets/CookCommandlet.cpp
+++ b/Engine/Source/Editor/UnrealEd/Private/Commandlets/CookCommandlet.cpp
@@ -25,6 +25,13 @@
 DEFINE_LOG_CATEGORY_STATIC(LogCookCommandlet, Log, All);
 
 
+UCookerSettings::UCookerSettings(const FPostConstructInitializeProperties& PCIP)
+	: Super(PCIP)
+{
+	DefaultPVRTCQuality = 0;
+}
+
+
 /* Static functions
  *****************************************************************************/
 
